Extract texture format indexing and interfaces from the TargetPlatformManager and instead put it into a new type:  TextureFormatManager instead.  TargetPlatformManager retains the same API for querying texture format information, it just forwards to the TextureFormatManager instead.  This is to facilitate the use of TextureCompressor from standalone executables that shouldn't depend on TargetPlatformManager.

diff --git a/Engine/Source/Developer/TargetPlatform/Private/TargetPlatformManagerModule.cpp b/Engine/Source/Developer/TargetPlatform/Private/TargetPlatformManagerModule.cpp
--- a/Engine/Source/Developer/TargetPlatform/Private/TargetPlatformManagerModule.cpp
+++ b/Engine/Source/Developer/TargetPlatform/Private/TargetPlatformManagerModule.cpp
@@ -17,6 +17,7 @@
 #include "Interfaces/IShaderFormat.h"
 #include "Interfaces/IShaderFormatModule.h"
 #include "Interfaces/ITextureFormat.h"
+#include "Interfaces/ITextureFormatManagerModule.h"
 #include "Interfaces/ITextureFormatModule.h"
 #include "PlatformInfo.h"
 #include "DesktopPlatformModule.h"
@@ -123,6 +124,8 @@ public:
 			}
 		}
 #endif
+		TextureFormatManager = FModuleManager::LoadModulePtr<ITextureFormatManagerModule>("TextureFormat");
+
 		// Calling a virtual function from a constructor, but with no expectation that a derived implementation of this
 		// method would be called.  This is solely to avoid duplicating code in this implementation, not for polymorphism.
 		FTargetPlatformManagerModule::Invalidate();
@@ -161,7 +164,7 @@ public:
 		if (!bHasInitErrors)
 		{
 			GetAudioFormats();
-			GetTextureFormats();
+			TextureFormatManager->Invalidate();
 			GetShaderFormats();
 		}
 
@@ -453,60 +456,12 @@ public:
 
 	virtual const TArray<const ITextureFormat*>& GetTextureFormats() override
 	{
-		static bool bInitialized = false;
-		static TArray<const ITextureFormat*> Results;
-
-		if (!bInitialized || bForceCacheUpdate)
-		{
-			bInitialized = true;
-			Results.Empty(Results.Num());
-
-			TArray<FName> Modules;
-
-			FModuleManager::Get().FindModules(TEXT("*TextureFormat*"), Modules);
-
-			if (!Modules.Num())
-			{
-				UE_LOG(LogTargetPlatformManager, Error, TEXT("No target texture formats found!"));
-			}
-
-			for (int32 Index = 0; Index < Modules.Num(); Index++)
-			{
-				ITextureFormatModule* Module = FModuleManager::LoadModulePtr<ITextureFormatModule>(Modules[Index]);
-				if (Module)
-				{
-					ITextureFormat* Format = Module->GetTextureFormat();
-					if (Format != nullptr)
-					{
-						Results.Add(Format);
-					}
-				}
-			}
-		}
-
-		return Results;
+		return TextureFormatManager->GetTextureFormats();
 	}
 
 	virtual const ITextureFormat* FindTextureFormat(FName Name) override
 	{
-		const TArray<const ITextureFormat*>& TextureFormats = GetTextureFormats();
-
-		for (int32 Index = 0; Index < TextureFormats.Num(); Index++)
-		{
-			TArray<FName> Formats;
-
-			TextureFormats[Index]->GetSupportedFormats(Formats);
-
-			for (int32 FormatIndex = 0; FormatIndex < Formats.Num(); FormatIndex++)
-			{
-				if (Formats[FormatIndex] == Name)
-				{
-					return TextureFormats[Index];
-				}
-			}
-		}
-
-		return nullptr;
+		return TextureFormatManager->FindTextureFormat(Name);
 	}
 
 	virtual const TArray<const IShaderFormat*>& GetShaderFormats() override
@@ -1207,6 +1162,9 @@ private:
 	// Map for fast lookup of platforms by name.
 	TMap<FName, ITargetPlatform*> PlatformsByName;
 
+	// External module that texture format operations are forwarded to
+	ITextureFormatManagerModule* TextureFormatManager;
+
 #if AUTOSDKS_ENABLED
 	// holds the list of Platforms that have attempted setup.
 	TMap<FName, bool> PlatformsSetup;
diff --git a/Engine/Source/Developer/TargetPlatform/Public/ChildTextureFormat.h b/Engine/Source/Developer/TargetPlatform/Public/ChildTextureFormat.h
deleted file mode 100644
diff --git a/Engine/Source/Developer/TargetPlatform/Public/Interfaces/ITextureFormat.h b/Engine/Source/Developer/TargetPlatform/Public/Interfaces/ITextureFormat.h
deleted file mode 100644
diff --git a/Engine/Source/Developer/TargetPlatform/Public/Interfaces/ITextureFormatModule.h b/Engine/Source/Developer/TargetPlatform/Public/Interfaces/ITextureFormatModule.h
deleted file mode 100644
diff --git a/Engine/Source/Developer/TargetPlatform/TargetPlatform.Build.cs b/Engine/Source/Developer/TargetPlatform/TargetPlatform.Build.cs
--- a/Engine/Source/Developer/TargetPlatform/TargetPlatform.Build.cs
+++ b/Engine/Source/Developer/TargetPlatform/TargetPlatform.Build.cs
@@ -18,6 +18,11 @@ public class TargetPlatform : ModuleRules
 		PublicDependencyModuleNames.Add("DesktopPlatform");
 		PublicDependencyModuleNames.Add("LauncherPlatform");
 
+		// TextureFormat contains public headers that were historically part of TargetPlatform, so it is exposed
+		// as a public include path on TargetPlatform.
+		PublicIncludePathModuleNames.Add("TextureFormat");
+		PublicDependencyModuleNames.Add("TextureFormat");
+
 		PrivateIncludePathModuleNames.Add("Engine");
 		PrivateIncludePathModuleNames.Add("PhysicsCore");
 
diff --git a/Engine/Source/Developer/TextureCompressor/Private/TextureCompressorModule.cpp b/Engine/Source/Developer/TextureCompressor/Private/TextureCompressorModule.cpp
--- a/Engine/Source/Developer/TextureCompressor/Private/TextureCompressorModule.cpp
+++ b/Engine/Source/Developer/TextureCompressor/Private/TextureCompressorModule.cpp
@@ -7,8 +7,8 @@
 #include "Async/AsyncWork.h"
 #include "Async/ParallelFor.h"
 #include "Modules/ModuleManager.h"
-#include "Engine/Texture.h"
-#include "Interfaces/ITargetPlatformManagerModule.h"
+#include "Engine/TextureDefines.h"
+#include "TextureFormatManager.h"
 #include "Interfaces/ITextureFormat.h"
 #include "Misc/Paths.h"
 #include "ImageCore.h"
@@ -2202,10 +2202,10 @@ public:
 	virtual bool UsesTaskGraph(const FTextureBuildSettings& BuildSettings) const override
 	{
 		const ITextureFormat* TextureFormat = nullptr;
-		ITargetPlatformManagerModule* TPM = GetTargetPlatformManager();
-		if (TPM)
+		ITextureFormatManagerModule* TFM = GetTextureFormatManager();
+		if (TFM)
 		{
-			TextureFormat = TPM->FindTextureFormat(BuildSettings.TextureFormatName);
+			TextureFormat = TFM->FindTextureFormat(BuildSettings.TextureFormatName);
 		}
 
 		if (TextureFormat)
@@ -2226,10 +2226,10 @@ public:
 		)
 	{
 		const ITextureFormat* TextureFormat = nullptr;
-		ITargetPlatformManagerModule* TPM = GetTargetPlatformManager();
-		if (TPM)
+		ITextureFormatManagerModule* TFM = GetTextureFormatManager();
+		if (TFM)
 		{
-			TextureFormat = TPM->FindTextureFormat(BuildSettings.TextureFormatName);
+			TextureFormat = TFM->FindTextureFormat(BuildSettings.TextureFormatName);
 		}
 		if (TextureFormat == nullptr)
 		{
@@ -2269,12 +2269,12 @@ public:
 
 			if(!BuildTextureMips(AssociatedNormalSourceMips, DefaultSettings, CompressorCaps, IntermediateAssociatedNormalSourceMipChain))
 			{
-				UE_LOG(LogTexture, Warning, TEXT("Failed to generate texture mips for composite texture"));
+				UE_LOG(LogTextureCompressor, Warning, TEXT("Failed to generate texture mips for composite texture"));
 			}
 
 			if(!ApplyCompositeTexture(IntermediateMipChain, IntermediateAssociatedNormalSourceMipChain, BuildSettings.CompositeTextureMode, BuildSettings.CompositePower))
 			{
-				UE_LOG(LogTexture, Warning, TEXT("Failed to apply composite texture"));
+				UE_LOG(LogTextureCompressor, Warning, TEXT("Failed to apply composite texture"));
 			}
 		}
 
diff --git a/Engine/Source/Developer/TextureCompressor/TextureCompressor.Build.cs b/Engine/Source/Developer/TextureCompressor/TextureCompressor.Build.cs
--- a/Engine/Source/Developer/TextureCompressor/TextureCompressor.Build.cs
+++ b/Engine/Source/Developer/TextureCompressor/TextureCompressor.Build.cs
@@ -6,13 +6,17 @@ public class TextureCompressor : ModuleRules
 {
 	public TextureCompressor(ReadOnlyTargetRules Target) : base(Target)
 	{
+		PrivateIncludePathModuleNames.AddRange(new string[]
+		{
+			"CoreUObject",
+			"Engine",
+		});
+
 		PrivateDependencyModuleNames.AddRange(
 			new string[] {
 				"Core",
-				"CoreUObject", // @todo Mac: for some reason it's needed to link in debug on Mac
-				"Engine",
-				"TargetPlatform",
-				"ImageCore"
+				"ImageCore",
+				"TextureFormat",
 			}
 			);
 
diff --git a/Engine/Source/Developer/TextureFormat/Private/TextureFormatManagerModule.cpp b/Engine/Source/Developer/TextureFormat/Private/TextureFormatManagerModule.cpp
new file mode 100644
--- /dev/null
+++ b/Engine/Source/Developer/TextureFormat/Private/TextureFormatManagerModule.cpp
@@ -0,0 +1,153 @@
+// Copyright Epic Games, Inc. All Rights Reserved.
+
+#include "CoreMinimal.h"
+#include "Interfaces/ITextureFormatManagerModule.h"
+#include "Interfaces/ITextureFormat.h"
+#include "Interfaces/ITextureFormatModule.h"
+#include "TextureFormatManager.h"
+
+DEFINE_LOG_CATEGORY_STATIC(LogTextureFormatManager, Log, All);
+
+/**
+ * Module for the target platform manager
+ */
+class FTextureFormatManagerModule
+	: public ITextureFormatManagerModule
+{
+public:
+
+	/** Default constructor. */
+	FTextureFormatManagerModule()
+		: ModuleName(TEXT("TextureFormat"))
+		, bForceCacheUpdate(true)
+		, bIgnoreFirstDelegateCall(true)
+	{
+		// Calling a virtual function from a constructor, but with no expectation that a derived implementation of this
+		// method would be called.  This is solely to avoid duplicating code in this implementation, not for polymorphism.
+		FTextureFormatManagerModule::Invalidate();
+
+		FModuleManager::Get().OnModulesChanged().AddRaw(this, &FTextureFormatManagerModule::ModulesChangesCallback);
+	}
+
+	/** Destructor. */
+	virtual ~FTextureFormatManagerModule()
+	{
+		FModuleManager::Get().OnModulesChanged().RemoveAll(this);
+	}
+
+
+	virtual const TArray<const ITextureFormat*>& GetTextureFormats() override
+	{
+		static bool bInitialized = false;
+		static TArray<const ITextureFormat*> Results;
+
+		if (!bInitialized || bForceCacheUpdate)
+		{
+			bInitialized = true;
+			Results.Empty(Results.Num());
+
+			TArray<FName> Modules;
+
+			FModuleManager::Get().FindModules(TEXT("*TextureFormat*"), Modules);
+
+			if (!Modules.Num())
+			{
+				UE_LOG(LogTextureFormatManager, Error, TEXT("No texture formats found!"));
+			}
+
+			for (int32 Index = 0; Index < Modules.Num(); Index++)
+			{
+				if (Modules[Index] != ModuleName) // Avoid our own module when going through this list that was gathered by name
+				{
+					ITextureFormatModule* Module = FModuleManager::LoadModulePtr<ITextureFormatModule>(Modules[Index]);
+					if (Module)
+					{
+						ITextureFormat* Format = Module->GetTextureFormat();
+						if (Format != nullptr)
+						{
+							Results.Add(Format);
+						}
+					}
+				}
+			}
+		}
+
+		return Results;
+	}
+
+	virtual const ITextureFormat* FindTextureFormat(FName Name) override
+	{
+		const TArray<const ITextureFormat*>& TextureFormats = GetTextureFormats();
+
+		for (int32 Index = 0; Index < TextureFormats.Num(); Index++)
+		{
+			TArray<FName> Formats;
+
+			TextureFormats[Index]->GetSupportedFormats(Formats);
+
+			for (int32 FormatIndex = 0; FormatIndex < Formats.Num(); FormatIndex++)
+			{
+				if (Formats[FormatIndex] == Name)
+				{
+					return TextureFormats[Index];
+				}
+			}
+		}
+
+		return nullptr;
+	}
+
+	virtual void Invalidate() override
+	{
+		bForceCacheUpdate = true;
+		GetTextureFormats();
+		bForceCacheUpdate = false;
+	}
+
+private:
+
+	void ModulesChangesCallback(FName InModuleName, EModuleChangeReason ReasonForChange)
+	{
+		if (!bIgnoreFirstDelegateCall && (InModuleName != ModuleName) && InModuleName.ToString().Contains(TEXT("TextureFormat")))
+		{
+			Invalidate();
+		}
+		bIgnoreFirstDelegateCall = false;
+	}
+
+	FName ModuleName;
+
+	// Flag to force reinitialization of all cached data. This is needed to have up-to-date caches
+	// in case of a module reload of a TextureFormat-Module.
+	bool bForceCacheUpdate;
+
+	// Flag to avoid redunant reloads
+	bool bIgnoreFirstDelegateCall;
+
+};
+
+IMPLEMENT_MODULE(FTextureFormatManagerModule, TextureFormat);
+
+class ITextureFormatManagerModule* GetTextureFormatManager()
+{
+	static class ITextureFormatManagerModule* SingletonInterface = NULL;
+	static bool bInitialized = false;
+	if (!bInitialized)
+	{
+		check(IsInGameThread());
+		bInitialized = true;
+		SingletonInterface = FModuleManager::LoadModulePtr<ITextureFormatManagerModule>("TextureFormat");
+	}
+	return SingletonInterface;
+}
+
+class ITextureFormatManagerModule& GetTextureFormatManagerRef()
+{
+	class ITextureFormatManagerModule* SingletonInterface = GetTextureFormatManager();
+	if (!SingletonInterface)
+	{
+		UE_LOG(LogInit, Fatal, TEXT("Texture format manager was requested, but not available."));
+		CA_ASSUME( SingletonInterface != NULL );	// Suppress static analysis warning in unreachable code (fatal error)
+	}
+	return *SingletonInterface;
+}
diff --git a/Engine/Source/Developer/TextureFormat/Public/ChildTextureFormat.h b/Engine/Source/Developer/TextureFormat/Public/ChildTextureFormat.h
new file mode 100644
--- /dev/null
+++ b/Engine/Source/Developer/TextureFormat/Public/ChildTextureFormat.h
@@ -0,0 +1,197 @@
+// Copyright Epic Games, Inc. All Rights Reserved.
+
+#pragma once
+
+#include "Interfaces/ITextureFormat.h"
+#include "Interfaces/ITextureFormatModule.h"
+#include "Interfaces/ITextureFormatManagerModule.h"
+#include "TextureCompressorModule.h"
+
+/**
+ * Version of ITextureFormat that handles a child texture format that is used as a "post-process" after compressing textures, useful for
+ * several platforms that need to modify already compressed texture data for optimal data
+ */
+class FChildTextureFormat : public ITextureFormat
+{
+public:
+
+	FChildTextureFormat(const TCHAR* PlatformFormatPrefix)
+		: FormatPrefix(PlatformFormatPrefix)
+	{
+
+	}
+
+protected:
+	void AddBaseTextureFormatModules(const TCHAR* ModuleNameWildcard)
+	{
+		TArray<FName> Modules;
+		FModuleManager::Get().FindModules(ModuleNameWildcard, Modules);
+
+		for (FName ModuleName : Modules)
+		{
+			ITextureFormatModule * TFModule = FModuleManager::LoadModulePtr<ITextureFormatModule>(ModuleName);
+			if ( TFModule != nullptr )
+			{
+				ITextureFormat* BaseFormat = TFModule->GetTextureFormat();
+				BaseFormat->GetSupportedFormats(BaseFormats);
+			}
+		}
+	}
+
+	FName GetBaseFormatName(FName PlatformName) const
+	{
+		return FName(*(PlatformName.ToString().Replace(*FormatPrefix, TEXT(""))));
+	}
+
+	/**
+	 * Given a platform specific format name, get the parent texture format object
+	 */
+	const ITextureFormat* GetBaseFormatObject(FName FormatName) const
+	{
+		FName BaseFormatName = GetBaseFormatName(FormatName);
+
+		ITextureFormatManagerModule& TFM = FModuleManager::LoadModuleChecked<ITextureFormatManagerModule>("TextureFormat");
+		const ITextureFormat* FormatObject = TFM.FindTextureFormat(BaseFormatName);
+
+		checkf(FormatObject != nullptr, TEXT("Bad FormatName %s passed to FChildTextureFormat::GetBaseFormatObject()"));
+
+		return FormatObject;
+	}
+
+	/**
+	 * The final version is a combination of parent and child formats, 8 bits for each
+	 */
+	virtual uint8 GetChildFormatVersion(FName Format, const struct FTextureBuildSettings* BuildSettings) const = 0;
+
+	/**
+	 * Make the child type think about if they need a key string or not, by making it pure virtual
+	 */
+	virtual FString GetChildDerivedDataKeyString(const class UTexture& Texture, const FTextureBuildSettings* BuildSettings) const = 0;
+
+
+public:
+
+	//// ITextureFormat interface ////
+
+	virtual void GetSupportedFormats(TArray<FName>& OutFormats) const override
+	{
+		for (FName BaseFormat : BaseFormats)
+		{
+			FName ChildFormat(*(FormatPrefix + BaseFormat.ToString()));
+			OutFormats.Add(ChildFormat);
+		}
+	}
+
+	virtual uint16 GetVersion(FName Format, const struct FTextureBuildSettings* BuildSettings) const final
+	{
+		uint16 BaseVersion = GetBaseFormatObject(Format)->GetVersion(Format, BuildSettings);
+		checkf(BaseVersion < 256, TEXT("BaseFormat for %s had too large a version (%d), must fit in 8bits"), *Format.ToString(), BaseVersion);
+
+		uint8 ChildVersion = GetChildFormatVersion(Format, BuildSettings);
+
+		// 8 bits for each version
+		return (BaseVersion << 8) | ChildVersion;
+	}
+
+	virtual FString GetDerivedDataKeyString(const class UTexture& Texture, const FTextureBuildSettings* BuildSettings) const final
+	{
+		FTextureBuildSettings BaseSettings = *BuildSettings;
+		BaseSettings.TextureFormatName = GetBaseFormatName(BuildSettings->TextureFormatName);
+
+		FString BaseString = GetBaseFormatObject(BuildSettings->TextureFormatName)->GetDerivedDataKeyString(Texture, &BaseSettings);
+		FString ChildString = GetChildDerivedDataKeyString(Texture, BuildSettings);
+
+		return BaseString + ChildString;
+	}
+
+	virtual EPixelFormat GetPixelFormatForImage(const struct FTextureBuildSettings& BuildSettings, const struct FImage& ExampleImage, bool bImageHasAlphaChannel) const override 
+	{
+		FTextureBuildSettings Settings = BuildSettings;
+		Settings.TextureFormatName = GetBaseFormatName(BuildSettings.TextureFormatName);
+		return GetBaseFormatObject(BuildSettings.TextureFormatName)->GetPixelFormatForImage(Settings, ExampleImage, bImageHasAlphaChannel);
+	}
+
+	bool CompressBaseImage(
+		const FImage& InImage,
+		const struct FTextureBuildSettings& BuildSettings,
+		bool bImageHasAlphaChannel,
+		FCompressedImage2D& OutCompressedImage
+	) const
+	{
+		FTextureBuildSettings BaseSettings = BuildSettings;
+		BaseSettings.TextureFormatName = GetBaseFormatName(BuildSettings.TextureFormatName);
+
+		// pass along the compression to the base format
+		if (GetBaseFormatObject(BuildSettings.TextureFormatName)->CompressImage(InImage, BaseSettings, bImageHasAlphaChannel, OutCompressedImage) == false)
+		{
+			UE_LOG(LogTemp, Error, TEXT("Failed to compress with base compressor [format %s]"), *BaseSettings.TextureFormatName.ToString());
+			return false;
+		}
+		return true;
+	}
+
+	bool CompressBaseImageTiled(
+		const FImage* Images,
+		uint32 NumImages,
+		const struct FTextureBuildSettings& BuildSettings,
+		bool bImageHasAlphaChannel,
+		TSharedPtr<FTilerSettings>& TilerSettings,
+		FCompressedImage2D& OutCompressedImage
+	) const
+	{
+		FTextureBuildSettings BaseSettings = BuildSettings;
+		BaseSettings.TextureFormatName = GetBaseFormatName(BuildSettings.TextureFormatName);
+
+		// pass along the compression to the base format
+		if (GetBaseFormatObject(BuildSettings.TextureFormatName)->CompressImageTiled(Images, NumImages, BaseSettings, bImageHasAlphaChannel, TilerSettings, OutCompressedImage) == false)
+		{
+			UE_LOG(LogTemp, Error, TEXT("Failed to compress with base tiled compressor [format %s]"), *BaseSettings.TextureFormatName.ToString());
+			return false;
+		}
+		return true;
+	}
+
+	bool PrepareTiling(
+		const FImage* Images,
+		const uint32 NumImages,
+		const struct FTextureBuildSettings& BuildSettings,
+		bool bImageHasAlphaChannel,
+		TSharedPtr<FTilerSettings>& OutTilerSettings,
+		TArray<FCompressedImage2D>& OutCompressedImages
+	) const override
+	{
+		FTextureBuildSettings BaseSettings = BuildSettings;
+		BaseSettings.TextureFormatName = GetBaseFormatName(BuildSettings.TextureFormatName);
+
+		return GetBaseFormatObject(BuildSettings.TextureFormatName)->PrepareTiling(Images, NumImages, BaseSettings, bImageHasAlphaChannel, OutTilerSettings, OutCompressedImages);
+	}
+
+	bool SetTiling(
+		const struct FTextureBuildSettings& BuildSettings,
+		TSharedPtr<FTilerSettings>& TilerSettings,
+		const TArray64<uint8>& ReorderedBlocks,
+		uint32 NumBlocks
+	) const override
+	{
+		FTextureBuildSettings BaseSettings = BuildSettings;
+		BaseSettings.TextureFormatName = GetBaseFormatName(BuildSettings.TextureFormatName);
+
+		return GetBaseFormatObject(BuildSettings.TextureFormatName)->SetTiling(BaseSettings, TilerSettings, ReorderedBlocks, NumBlocks);
+	}
+
+	void ReleaseTiling(const struct FTextureBuildSettings& BuildSettings, TSharedPtr<FTilerSettings>& TilerSettings) const override
+	{
+		FTextureBuildSettings BaseSettings = BuildSettings;
+		BaseSettings.TextureFormatName = GetBaseFormatName(BuildSettings.TextureFormatName);
+
+		return GetBaseFormatObject(BuildSettings.TextureFormatName)->ReleaseTiling(BuildSettings, TilerSettings);
+	}
+
+protected:
+
+	// Prefix put before all formats from parent formats
+	FString FormatPrefix;
+
+	// List of base formats that. Combined with FormatPrefix, this contains all formats this can handle
+	TArray<FName> BaseFormats;
+};
diff --git a/Engine/Source/Developer/TextureFormat/Public/Interfaces/ITextureFormat.h b/Engine/Source/Developer/TextureFormat/Public/Interfaces/ITextureFormat.h
new file mode 100644
--- /dev/null
+++ b/Engine/Source/Developer/TextureFormat/Public/Interfaces/ITextureFormat.h
@@ -0,0 +1,245 @@
+// Copyright Epic Games, Inc. All Rights Reserved.
+
+#pragma once
+
+#include "CoreMinimal.h"
+#include "PixelFormat.h"
+/**
+ * Structure for texture format compressor capabilities.
+ */
+struct FTextureFormatCompressorCaps
+{
+	FTextureFormatCompressorCaps()
+		: MaxTextureDimension(TNumericLimits<uint32>::Max())
+		, NumMipsInTail(0)
+		, ExtData(0)
+	{ }
+
+	uint32 MaxTextureDimension;
+	uint32 NumMipsInTail;
+	uint32 ExtData;
+};
+
+
+/**
+ * Interface for texture compression modules.
+ */
+class ITextureFormat
+{
+public:
+
+	/**
+	 * Checks whether this texture format can compress in parallel.
+	 *
+	 * @return true if parallel compression is supported, false otherwise.
+	 */
+	virtual bool AllowParallelBuild() const
+	{
+		return false;
+	}
+
+	/**
+	 * Gets the current version of the specified texture format.
+	 *
+	 * @param Format The format to get the version for.
+	 * @return Version number.
+	 */
+	virtual uint16 GetVersion(
+		FName Format,
+		const struct FTextureBuildSettings* BuildSettings = nullptr
+	) const = 0;
+
+	/**
+	 * Gets an optional derived data key string, so that the compressor can
+	 * rely upon the number of mips, size of texture, etc, when compressing the image
+	 *
+	 * @param Texture Reference to the texture we are compressing.
+	 * @return A string that will be used with the DDC, the string should be in the format "<DATA>_"
+	 */
+	virtual FString GetDerivedDataKeyString( const class UTexture& Texture, const FTextureBuildSettings* BuildSettings) const
+	{
+		return TEXT("");
+	}
+
+	/**
+	 * Gets the list of supported formats.
+	 *
+	 * @param OutFormats Will hold the list of formats.
+	 */
+	virtual void GetSupportedFormats( TArray<FName>& OutFormats ) const = 0;
+
+	/**
+	* Gets the capabilities of the texture compressor.
+	*
+	* @param OutCaps Filled with capability properties of texture format compressor.
+	*/
+	virtual FTextureFormatCompressorCaps GetFormatCapabilities() const = 0;
+
+	/**
+	* Gets the capabilities of the texture compressor.
+	*
+	* @param OutCaps Filled with capability properties of texture format compressor.
+	*/
+	virtual FTextureFormatCompressorCaps GetFormatCapabilitiesEx(const struct FTextureBuildSettings& BuildSettings, uint32 NumMips, const struct FImage& ExampleImage, bool bImageHasAlphaChannel) const
+	{
+		return GetFormatCapabilities();
+	}
+
+	/**
+	 * Calculate the final/runtime pixel format for this image on this platform
+	 */
+	virtual EPixelFormat GetPixelFormatForImage(const struct FTextureBuildSettings& BuildSettings, const struct FImage& Image, bool bImageHasAlphaChannel) const = 0;
+
+	/**
+	 * Compresses a single image.
+	 *
+	 * @param Image The input image.
+	 * @param BuildSettings Build settings.
+	 * @param bImageHasAlphaChannel true if the image has a non-white alpha channel.
+	 * @param OutCompressedMip The compressed image.
+	 * @returns true on success, false otherwise.
+	 */
+	virtual bool CompressImage(
+		const struct FImage& Image,
+		const struct FTextureBuildSettings& BuildSettings,
+		bool bImageHasAlphaChannel,
+		struct FCompressedImage2D& OutCompressedImage
+	) const = 0;
+
+	/**
+	 * Compress an image (or images for a miptail) into a single mip blob.
+	 *
+	 * @param Images The input image(s)
+	 * @param Images The number of images (for a miptail, this number should match what was returned in GetFormatCapabilities, mostly used for verification)
+	 * @param BuildSettings Build settings.
+	 * @param bImageHasAlphaChannel true if the image has a non-white alpha channel.
+	 * @param ExtData Extra data that the format may want to have passed back in to each compress call (makes the format class be stateless)
+	 * @param OutCompressedMip The compressed image.
+	 * @returns true on success, false otherwise.
+	 */
+	virtual bool CompressImageEx(const struct FImage* Images,
+		const uint32 NumImages,
+		const struct FTextureBuildSettings& BuildSettings,
+		bool bImageHasAlphaChannel,
+		uint32 ExtData,
+		FCompressedImage2D& OutCompressedImage) const
+	{
+		// general case can't handle mip tails
+		if (Images == nullptr || NumImages > 1)
+		{
+			return false;
+		}
+		
+		return CompressImage(*Images, BuildSettings, bImageHasAlphaChannel, OutCompressedImage);
+	}
+
+	/**
+	 * An object produced by PrepareTiling and used by SetTiling and CompressImageTiled.
+	 */
+	struct FTilerSettings
+	{
+	};
+
+	/**
+	 * Compress an image (or images for a miptail) into a single mip blob with device-specific tiling.
+	 *
+	 * @param Image The input image.
+	 * @param BuildSettings Build settings.
+	 * @param bImageHasAlphaChannel true if the image has a non-white alpha channel.
+	 * @param OutCompressedMip The compressed image.
+	 * @param Tiler The tiler settings.
+	 * @returns true on success, false otherwise.
+	 */
+	virtual bool CompressImageTiled(
+		const struct FImage* Images,
+		uint32 NumImages,
+		const struct FTextureBuildSettings& BuildSettings,
+		bool bImageHasAlphaChannel,
+		TSharedPtr<FTilerSettings>& TilerSettings,
+		struct FCompressedImage2D& OutCompressedImage) const
+	{
+		unimplemented();
+		return false;
+	}
+
+	/**
+	 * Whether device-specific tiling is supported by the compressor.
+	 *
+	 * @param BuildSettings Build settings.
+	 * @returns true if tiling is supported, false if it must be done by the caller
+	 *
+	 */
+	virtual bool SupportsTiling(const struct FTextureBuildSettings& BuildSettings) const
+	{
+		return false;
+	}
+
+	/**
+	 * Prepares to compresses a single image with tiling. The result OutTilerSettings is used by SetTiling and CompressImageTiled.
+	 *
+	 * @param Image The input image.
+	 * @param BuildSettings Build settings.
+	 * @param bImageHasAlphaChannel true if the image has a non-white alpha channel.
+	 * @param OutTilerSettings The tiler settings that will be used by CompressImageTiled and SetTiling.
+	 * @param OutCompressedImage The image to tile.
+	 * @returns true on success, false otherwise.
+	 */
+	virtual bool PrepareTiling(
+		const FImage* Images,
+		const uint32 NumImages,
+		const struct FTextureBuildSettings& BuildSettings,
+		bool bImageHasAlphaChannel,
+		TSharedPtr<FTilerSettings>& OutTilerSettings,
+		TArray<FCompressedImage2D>& OutCompressedImage
+	) const
+	{
+		unimplemented();
+		return false;
+	}
+	
+	/**
+	 * Sets the tiling settings after device-specific tiling has been performed.
+	 *
+	 * @param BuildSettings Build settings.
+	 * @param TilerSettings The tiler settings produced by PrepareTiling.
+	 * @param ReorderedBlocks The blocks that have been tiled.
+	 * @param NumBlocks The number of blocks.
+	 * @returns true on success, false otherwise.
+	 */
+	virtual bool SetTiling(
+		const struct FTextureBuildSettings& BuildSettings,
+		TSharedPtr<FTilerSettings>& TilerSettings,
+		const TArray64<uint8>& ReorderedBlocks,
+		uint32 NumBlocks
+	) const
+	{
+		unimplemented();
+		return false;
+	}
+
+	/**
+	 * Cleans up the FTilerSettings object once it is finished.
+	 * 
+	 * @param BuildSettings Build settings.
+	 * @param TilerSettings The tiler settings object to release.
+	 */
+	virtual void ReleaseTiling(const struct FTextureBuildSettings& BuildSettings, TSharedPtr<FTilerSettings>& TilerSettings) const
+	{
+		unimplemented();
+	}
+
+	/**
+	 * Whether the compressor uses the FTaskGraph API.
+	 * 
+	 * @returns true if FTaskGraph is used, false otherwise
+	 */
+	virtual bool UsesTaskGraph() const
+	{
+		return false;
+	}
+
+public:
+
+	/** Virtual destructor. */
+	virtual ~ITextureFormat() { }
+};
diff --git a/Engine/Source/Developer/TextureFormat/Public/Interfaces/ITextureFormatManagerModule.h b/Engine/Source/Developer/TextureFormat/Public/Interfaces/ITextureFormatManagerModule.h
new file mode 100644
--- /dev/null
+++ b/Engine/Source/Developer/TextureFormat/Public/Interfaces/ITextureFormatManagerModule.h
@@ -0,0 +1,42 @@
+// Copyright Epic Games, Inc. All Rights Reserved.
+
+#pragma once
+
+#include "CoreMinimal.h"
+#include "Modules/ModuleInterface.h"
+
+/**
+ * Module for the texture format manager
+ */
+class ITextureFormatManagerModule
+	: public IModuleInterface
+{
+public:
+
+	/**
+	 * Finds a texture format with the specified name.
+	 *
+	 * @param Name Name of the format to find.
+	 * @return The texture format, or nullptr if not found.
+	 */
+	virtual const class ITextureFormat* FindTextureFormat( FName Name ) = 0;
+
+	/**
+	 * Returns the list of all ITextureFormats that were located in DLLs.
+	 *
+	 * @return Collection of texture formats.
+	 */
+	virtual const TArray<const class ITextureFormat*>& GetTextureFormats() = 0;
+
+	/**
+	 * Invalidates the texture format manager module.
+	 *
+	 * Invalidate should be called if any TextureFormat modules get loaded/unloaded/reloaded during 
+	 * runtime to give the implementation the chance to rebuild all its internal states and caches.
+	 */
+	virtual void Invalidate() = 0;
+public:
+
+	/** Virtual destructor. */
+	~ITextureFormatManagerModule() { }
+};
diff --git a/Engine/Source/Developer/TextureFormat/Public/Interfaces/ITextureFormatModule.h b/Engine/Source/Developer/TextureFormat/Public/Interfaces/ITextureFormatModule.h
new file mode 100644
--- /dev/null
+++ b/Engine/Source/Developer/TextureFormat/Public/Interfaces/ITextureFormatModule.h
@@ -0,0 +1,29 @@
+// Copyright Epic Games, Inc. All Rights Reserved.
+
+#pragma once
+
+#include "CoreMinimal.h"
+#include "Modules/ModuleInterface.h"
+
+class ITextureFormat;
+
+/**
+ * Interface for texture format modules.
+ */
+class ITextureFormatModule
+	: public IModuleInterface
+{
+public:
+
+	/**
+	 * Gets the texture format.
+	 *
+	 * @return The texture format interface.
+	 */
+	virtual ITextureFormat* GetTextureFormat() = 0;
+
+public:
+
+	/** Virtual destructor. */
+	~ITextureFormatModule() { }
+};
diff --git a/Engine/Source/Developer/TextureFormat/Public/TextureFormatManager.h b/Engine/Source/Developer/TextureFormat/Public/TextureFormatManager.h
new file mode 100644
--- /dev/null
+++ b/Engine/Source/Developer/TextureFormat/Public/TextureFormatManager.h
@@ -0,0 +1,12 @@
+// Copyright Epic Games, Inc. All Rights Reserved.
+
+#pragma once
+
+#include "CoreTypes.h"
+#include "Interfaces/ITextureFormatManagerModule.h"
+
+/** Return the Texture Format Manager interface, if it is available, otherwise return nullptr. **/
+TEXTUREFORMAT_API class ITextureFormatManagerModule* GetTextureFormatManager();
+
+/** Return the Texture Format Manager interface, fatal error if it is not available. **/
+TEXTUREFORMAT_API class ITextureFormatManagerModule& GetTextureFormatManagerRef();
diff --git a/Engine/Source/Developer/TextureFormat/TextureFormat.Build.cs b/Engine/Source/Developer/TextureFormat/TextureFormat.Build.cs
new file mode 100644
--- /dev/null
+++ b/Engine/Source/Developer/TextureFormat/TextureFormat.Build.cs
@@ -0,0 +1,18 @@
+// Copyright Epic Games, Inc. All Rights Reserved.
+
+using UnrealBuildTool;
+
+public class TextureFormat : ModuleRules
+{
+	public TextureFormat(ReadOnlyTargetRules Target) : base(Target)
+	{
+		PrivateDependencyModuleNames.AddRange(
+			new string[] {
+				"Core",
+				"ImageCore"
+			}
+			);
+
+		AddEngineThirdPartyPrivateStaticDependencies(Target, "nvTextureTools");
+	}
+}
diff --git a/Engine/Source/Runtime/Engine/Classes/Engine/Texture.h b/Engine/Source/Runtime/Engine/Classes/Engine/Texture.h
--- a/Engine/Source/Runtime/Engine/Classes/Engine/Texture.h
+++ b/Engine/Source/Runtime/Engine/Classes/Engine/Texture.h
@@ -71,25 +71,6 @@ enum TextureAddress
 	TA_MAX,
 };
 
-UENUM()
-enum ECompositeTextureMode
-{
-	CTM_Disabled UMETA(DisplayName="Disabled"),
-	/** CompositingTexture needs to be a normal map with the same or larger size. */
-	CTM_NormalRoughnessToRed UMETA(DisplayName="Add Normal Roughness To Red"),
-	/** CompositingTexture needs to be a normal map with the same or larger size. */
-	CTM_NormalRoughnessToGreen UMETA(DisplayName="Add Normal Roughness To Green"),
-	/** CompositingTexture needs to be a normal map with the same or larger size. */
-	CTM_NormalRoughnessToBlue UMETA(DisplayName="Add Normal Roughness To Blue"),
-	/** CompositingTexture needs to be a normal map with the same or larger size. */
-	CTM_NormalRoughnessToAlpha UMETA(DisplayName="Add Normal Roughness To Alpha"),
-	CTM_MAX,
-
-	// Note: These are serialized as as raw values in the texture DDC key, so additional entries
-	// should be added at the bottom; reordering or removing entries will require changing the GUID
-	// in the texture compressor DDC key
-};
-
 UENUM()
 enum ETextureMipCount
 {
diff --git a/Engine/Source/Runtime/Engine/Classes/Engine/TextureDefines.h b/Engine/Source/Runtime/Engine/Classes/Engine/TextureDefines.h
--- a/Engine/Source/Runtime/Engine/Classes/Engine/TextureDefines.h
+++ b/Engine/Source/Runtime/Engine/Classes/Engine/TextureDefines.h
@@ -4,7 +4,14 @@
 
 #include "CoreMinimal.h"
 #include "UObject/ObjectMacros.h"
+
+// TextureDefines.h is used from the TextureCompressor module as an "include only"
+// dependency to get visibility of these enum values without linking to the engine module
+// to facilitate that the generated TextureDefines header is only conditionally included
+// if we're compiling with the engine.
+#if WITH_ENGINE
 #include "TextureDefines.generated.h"
+#endif // WITH_ENGINE
 
 /**
  * @warning: if this is changed:
@@ -216,3 +223,23 @@ enum ETextureLossyCompressionAmount
 	TLCA_High			UMETA(DisplayName = "High"),
 	TLCA_Highest		UMETA(DisplayName = "Highest (Worst image quality, smallest filesize)"),
 };
+
+
+UENUM()
+enum ECompositeTextureMode
+{
+	CTM_Disabled UMETA(DisplayName="Disabled"),
+	/** CompositingTexture needs to be a normal map with the same or larger size. */
+	CTM_NormalRoughnessToRed UMETA(DisplayName="Add Normal Roughness To Red"),
+	/** CompositingTexture needs to be a normal map with the same or larger size. */
+	CTM_NormalRoughnessToGreen UMETA(DisplayName="Add Normal Roughness To Green"),
+	/** CompositingTexture needs to be a normal map with the same or larger size. */
+	CTM_NormalRoughnessToBlue UMETA(DisplayName="Add Normal Roughness To Blue"),
+	/** CompositingTexture needs to be a normal map with the same or larger size. */
+	CTM_NormalRoughnessToAlpha UMETA(DisplayName="Add Normal Roughness To Alpha"),
+	CTM_MAX,
+
+	// Note: These are serialized as as raw values in the texture DDC key, so additional entries
+	// should be added at the bottom; reordering or removing entries will require changing the GUID
+	// in the texture compressor DDC key
+};
