Allow GetTextureFormatManager to be called for the first time from any thread.
Fix NonUnity issue in TextureFormatManager.

diff --git a/Engine/Source/Developer/TextureFormat/Private/TextureFormatManagerModule.cpp b/Engine/Source/Developer/TextureFormat/Private/TextureFormatManagerModule.cpp
--- a/Engine/Source/Developer/TextureFormat/Private/TextureFormatManagerModule.cpp
+++ b/Engine/Source/Developer/TextureFormat/Private/TextureFormatManagerModule.cpp
@@ -4,6 +4,7 @@
 #include "Interfaces/ITextureFormatManagerModule.h"
 #include "Interfaces/ITextureFormat.h"
 #include "Interfaces/ITextureFormatModule.h"
+#include "Modules/ModuleManager.h"
 #include "TextureFormatManager.h"
 
 DEFINE_LOG_CATEGORY_STATIC(LogTextureFormatManager, Log, All);
@@ -130,14 +131,7 @@ IMPLEMENT_MODULE(FTextureFormatManagerModule, TextureFormat);
 
 class ITextureFormatManagerModule* GetTextureFormatManager()
 {
-	static class ITextureFormatManagerModule* SingletonInterface = NULL;
-	static bool bInitialized = false;
-	if (!bInitialized)
-	{
-		check(IsInGameThread());
-		bInitialized = true;
-		SingletonInterface = FModuleManager::LoadModulePtr<ITextureFormatManagerModule>("TextureFormat");
-	}
+	static class ITextureFormatManagerModule* SingletonInterface = FModuleManager::LoadModulePtr<ITextureFormatManagerModule>("TextureFormat");
 	return SingletonInterface;
 }
 
