Preferred alternative fix (instead of CL16332153) for UE-115783 and UE-115714 that uses inlining instead of relying on the caller to selectively enable/disable linking in their build files.  This leaves the TextureFormat module with zero exports.  Only interfaces and two inline functions.  Confirmed compile & link on: win64 editor, win64 server shipping, linux server shipping, ShooterGame linux editor

diff --git a/Engine/Source/Developer/TextureFormat/Private/TextureFormatManagerModule.cpp b/Engine/Source/Developer/TextureFormat/Private/TextureFormatManagerModule.cpp
--- a/Engine/Source/Developer/TextureFormat/Private/TextureFormatManagerModule.cpp
+++ b/Engine/Source/Developer/TextureFormat/Private/TextureFormatManagerModule.cpp
@@ -153,20 +153,3 @@ private:
 };
 
 IMPLEMENT_MODULE(FTextureFormatManagerModule, TextureFormat);
-
-class ITextureFormatManagerModule* GetTextureFormatManager()
-{
-	static class ITextureFormatManagerModule* SingletonInterface = FModuleManager::LoadModulePtr<ITextureFormatManagerModule>("TextureFormat");
-	return SingletonInterface;
-}
-
-class ITextureFormatManagerModule& GetTextureFormatManagerRef()
-{
-	class ITextureFormatManagerModule* SingletonInterface = GetTextureFormatManager();
-	if (!SingletonInterface)
-	{
-		UE_LOG(LogInit, Fatal, TEXT("Texture format manager was requested, but not available."));
-		CA_ASSUME( SingletonInterface != NULL );	// Suppress static analysis warning in unreachable code (fatal error)
-	}
-	return *SingletonInterface;
-}
diff --git a/Engine/Source/Developer/TextureFormat/Public/TextureFormatManager.h b/Engine/Source/Developer/TextureFormat/Public/TextureFormatManager.h
--- a/Engine/Source/Developer/TextureFormat/Public/TextureFormatManager.h
+++ b/Engine/Source/Developer/TextureFormat/Public/TextureFormatManager.h
@@ -6,7 +6,20 @@
 #include "Interfaces/ITextureFormatManagerModule.h"
 
 /** Return the Texture Format Manager interface, if it is available, otherwise return nullptr. **/
-TEXTUREFORMAT_API class ITextureFormatManagerModule* GetTextureFormatManager();
+inline ITextureFormatManagerModule* GetTextureFormatManager()
+{
+	return FModuleManager::LoadModulePtr<ITextureFormatManagerModule>("TextureFormat");
+}
 
 /** Return the Texture Format Manager interface, fatal error if it is not available. **/
-TEXTUREFORMAT_API class ITextureFormatManagerModule& GetTextureFormatManagerRef();
+inline ITextureFormatManagerModule& GetTextureFormatManagerRef()
+{
+	class ITextureFormatManagerModule* TextureFormatManager = GetTextureFormatManager();
+	if (!TextureFormatManager)
+	{
+		UE_LOG(LogInit, Fatal, TEXT("Texture format manager was requested, but not available."));
+		CA_ASSUME( TextureFormatManager != NULL );	// Suppress static analysis warning in unreachable code (fatal error)
+	}
+	return *TextureFormatManager;
+}
+
diff --git a/Engine/Source/Runtime/Engine/Engine.Build.cs b/Engine/Source/Runtime/Engine/Engine.Build.cs
--- a/Engine/Source/Runtime/Engine/Engine.Build.cs
+++ b/Engine/Source/Runtime/Engine/Engine.Build.cs
@@ -151,7 +151,6 @@ public class Engine : ModuleRules
 			PrivateDependencyModuleNames.AddRange(
 				new string[] {
 					"RawMesh",
-					"TextureFormat",
 				}
 			);
 		}
