Fix non-unity compile issue in TextureFormatManager

diff --git a/Engine/Source/Developer/TextureFormat/Public/TextureFormatManager.h b/Engine/Source/Developer/TextureFormat/Public/TextureFormatManager.h
--- a/Engine/Source/Developer/TextureFormat/Public/TextureFormatManager.h
+++ b/Engine/Source/Developer/TextureFormat/Public/TextureFormatManager.h
@@ -4,6 +4,7 @@
 
 #include "CoreTypes.h"
 #include "Interfaces/ITextureFormatManagerModule.h"
+#include "Modules/ModuleManager.h"
 
 /** Return the Texture Format Manager interface, if it is available, otherwise return nullptr. **/
 inline ITextureFormatManagerModule* GetTextureFormatManager()
