Cache texture formats instead of recreating FName from strings every call to GetDerivedTextureDataKey

diff --git a/Engine/Source/Developer/TextureFormat/Public/ChildTextureFormat.h b/Engine/Source/Developer/TextureFormat/Public/ChildTextureFormat.h
--- a/Engine/Source/Developer/TextureFormat/Public/ChildTextureFormat.h
+++ b/Engine/Source/Developer/TextureFormat/Public/ChildTextureFormat.h
@@ -38,6 +38,13 @@ protected:
 				BaseFormat->GetSupportedFormats(BaseFormats);
 			}
 		}
+
+		SupportedFormatsCached.Reset();
+		for (FName BaseFormat : BaseFormats)
+		{
+			FName ChildFormat(*(FormatPrefix + BaseFormat.ToString()));
+			SupportedFormatsCached.Add(ChildFormat);
+		}
 	}
 
 	FName GetBaseFormatName(FName PlatformName) const
@@ -111,11 +118,7 @@ public:
 
 	virtual void GetSupportedFormats(TArray<FName>& OutFormats) const override
 	{
-		for (FName BaseFormat : BaseFormats)
-		{
-			FName ChildFormat(*(FormatPrefix + BaseFormat.ToString()));
-			OutFormats.Add(ChildFormat);
-		}
+		OutFormats.Append(SupportedFormatsCached);
 	}
 
 	virtual bool SupportsEncodeSpeed(FName Format) const override
@@ -262,8 +265,11 @@ public:
 protected:
 
 	// Prefix put before all formats from parent formats
-	FString FormatPrefix;
+	const FString FormatPrefix;
 
 	// List of base formats that. Combined with FormatPrefix, this contains all formats this can handle
 	TArray<FName> BaseFormats;
+
+	// List of combined BaseFormats with FormatPrefix.
+	TArray<FName> SupportedFormatsCached;
 };
