TextureFormatManager better init safety
do Invalidate in your constructor so you are valid after LoadModule
in case someone tries to Load TextureFormatManager and use it other than TargetPlatform
mutex protect thread access
fixes TBW broken by previous CL on TextureFormatManager

diff --git a/Engine/Source/Developer/TargetPlatform/Private/TargetPlatformManagerModule.cpp b/Engine/Source/Developer/TargetPlatform/Private/TargetPlatformManagerModule.cpp
--- a/Engine/Source/Developer/TargetPlatform/Private/TargetPlatformManagerModule.cpp
+++ b/Engine/Source/Developer/TargetPlatform/Private/TargetPlatformManagerModule.cpp
@@ -135,6 +135,7 @@ public:
 		, bForceCacheUpdate(true)
 		, bHasInitErrors(false)
 		, bIgnoreFirstDelegateCall(true)
+		, bSkipOneTextureFormatManagerInvalidate(false)
 	{
 #if WITH_EDITOR
 
@@ -164,12 +165,11 @@ public:
 			}
 		}
 #endif
+
 		TextureFormatManager = FModuleManager::LoadModulePtr<ITextureFormatManagerModule>("TextureFormat");
 
-		//TextureFormatManager must have ->Invalidate called before it can be used
-		//TextureFormatManager->Invalidate();
-		// this is done by FTargetPlatformManagerModule->Invalidate() below
-		// FTargetPlatformManagerModule is not fully constructed yet so calls to it are dangerous
+		//TextureFormatManager->Invalidate() already done, don't do again now :
+		bSkipOneTextureFormatManagerInvalidate = true;
 
 		// Calling a virtual function from a constructor, but with no expectation that a derived implementation of this
 		// method would be called.  This is solely to avoid duplicating code in this implementation, not for polymorphism.
@@ -206,12 +206,20 @@ public:
 		GetActiveTargetPlatforms();
 
 		bForceCacheUpdate = false;
+		
+		if ( bSkipOneTextureFormatManagerInvalidate )
+		{
+			bSkipOneTextureFormatManagerInvalidate = false;
+		}
+		else if (!bHasInitErrors)
+		{
+			TextureFormatManager->Invalidate();
+		}
 
 		// If we've had an error due to an invalid target platform, don't do additional work
 		if (!bHasInitErrors)
 		{
 			GetAudioFormats();
-			TextureFormatManager->Invalidate();
 			GetShaderFormats();
 		}
 
@@ -1307,6 +1315,9 @@ private:
 
 	// Flag to avoid redunant reloads
 	bool bIgnoreFirstDelegateCall;
+	
+	// Flag to avoid redunant reloads
+	bool bSkipOneTextureFormatManagerInvalidate;
 
 	// Holds the list of discovered platforms.
 	TArray<ITargetPlatform*> Platforms;
diff --git a/Engine/Source/Developer/TextureFormat/Private/TextureFormatManagerModule.cpp b/Engine/Source/Developer/TextureFormat/Private/TextureFormatManagerModule.cpp
--- a/Engine/Source/Developer/TextureFormat/Private/TextureFormatManagerModule.cpp
+++ b/Engine/Source/Developer/TextureFormat/Private/TextureFormatManagerModule.cpp
@@ -6,6 +6,8 @@
 #include "Interfaces/ITextureFormatModule.h"
 #include "Modules/ModuleManager.h"
 #include "TextureFormatManager.h"
+#include "HAL/CriticalSection.h"
+#include "Misc/ScopeLock.h"
 
 DEFINE_LOG_CATEGORY_STATIC(LogTextureFormatManager, Log, All);
 
@@ -33,9 +35,12 @@ public:
 		, bModuleChangeCallbackEnabled(false)
 		, TextureFormatsInitPhase(EInitPhase::JustConstructedNotInit)
 	{
+		// Calling a virtual function from a constructor, but with no expectation that a derived implementation of this
+		// method would be called.  This is solely to avoid duplicating code in this implementation, not for polymorphism.
+		FTextureFormatManagerModule::Invalidate();
+		
+		// add AFTER Invalidate :
 		FModuleManager::Get().OnModulesChanged().AddRaw(this, &FTextureFormatManagerModule::ModulesChangesCallback);
-
-		// not usable until first Invalidate is called
 	}
 
 	/** Destructor. */
@@ -46,9 +51,10 @@ public:
 		FModuleManager::Get().OnModulesChanged().RemoveAll(this);
 	}
 
-
 	virtual const TArray<const ITextureFormat*>& GetTextureFormats() override
 	{
+		FScopeLock Lock(&ModuleMutex);
+
 		// should not be called recursively while I am building the list :
 		check( TextureFormatsInitPhase!= EInitPhase::GetTextureFormatsInProgressDontTouch );
 
@@ -120,7 +126,7 @@ public:
 			TextureFormatsInitPhase = EInitPhase::GetTextureFormatsPartialOkayToRead;
 
 			// run through the Child formats and call GetTextureFormat() on them
-			// this will call back to me and do GetTextureFormats() which will get only the base formats
+			// this could call back to me and do GetTextureFormats() which will get only the base formats
 			for (int32 Index = 0; Index < ChildModules.Num(); Index++)
 			{
 				ITextureFormatModule* Module = ChildModules[Index].Module;
@@ -160,11 +166,10 @@ public:
 
 		return TextureFormats;
 	}
-
+	
 	virtual const ITextureFormat* FindTextureFormat(FName Name) override
 	{
-		check( (int)TextureFormatsInitPhase >= (int)EInitPhase::GetTextureFormatsPartialOkayToRead );
-
+		// just pass through to FindTextureFormatAndModule
 		FName ModuleNameUnused;
 		ITextureFormatModule* ModuleUnused;
 		return FindTextureFormatAndModule(Name, ModuleNameUnused, ModuleUnused);
@@ -172,12 +177,13 @@ public:
 	
 	virtual const class ITextureFormat* FindTextureFormatAndModule(FName Name, FName& OutModuleName, ITextureFormatModule*& OutModule) override
 	{
+		FScopeLock Lock(&ModuleMutex);
 		check( (int)TextureFormatsInitPhase >= (int)EInitPhase::GetTextureFormatsPartialOkayToRead );
 
 		// Called to ensure the arrays are populated
 		// dangerous and not necessary, removed :
-		check( ! bForceCacheUpdate );
 		//GetTextureFormats();
+		check( ! bForceCacheUpdate );
 
 		for (int32 Index = 0; Index < TextureFormats.Num(); Index++)
 		{
@@ -202,8 +208,8 @@ public:
 
 	virtual void Invalidate() override
 	{
-		// this is called right after the constructor
-		// we are not usable until the first Invalidate is called
+		FScopeLock Lock(&ModuleMutex);
+		// this is called from the constructor
 		TextureFormatsInitPhase = EInitPhase::Invalidated;
 		bForceCacheUpdate = true;
 		GetTextureFormats();
@@ -213,6 +219,7 @@ private:
 
 	void ModulesChangesCallback(FName InModuleName, EModuleChangeReason ReasonForChange)
 	{
+		FScopeLock Lock(&ModuleMutex);
 		if (bModuleChangeCallbackEnabled && (InModuleName != ModuleName) && InModuleName.ToString().Contains(TEXT("TextureFormat")))
 		{
 			// when a "TextureFormat" module is loaded, rebuild my list
@@ -241,8 +248,7 @@ private:
 	// Track tricky initialization progress
 	EInitPhase TextureFormatsInitPhase;
 
-	// @@!! Mutex and ScopeLock everything in case this is called from multiple threads ?
-
+	FCriticalSection ModuleMutex;
 };
 
 IMPLEMENT_MODULE(FTextureFormatManagerModule, TextureFormat);
diff --git a/Engine/Source/Developer/TextureFormat/Public/ChildTextureFormat.h b/Engine/Source/Developer/TextureFormat/Public/ChildTextureFormat.h
--- a/Engine/Source/Developer/TextureFormat/Public/ChildTextureFormat.h
+++ b/Engine/Source/Developer/TextureFormat/Public/ChildTextureFormat.h
@@ -9,6 +9,7 @@
 #include "Serialization/CompactBinaryWriter.h"
 #include "TextureCompressorModule.h"
 #include "Algo/Unique.h"
+#include "TextureFormatManager.h"
 
 /**
  * Version of ITextureFormat that handles a child texture format that is used as a "post-process" after compressing textures, useful for
@@ -82,7 +83,7 @@ protected:
 	{
 		FName BaseFormatName = GetBaseFormatName(FormatName);
 
-		ITextureFormatManagerModule& TFM = FModuleManager::LoadModuleChecked<ITextureFormatManagerModule>("TextureFormat");
+		ITextureFormatManagerModule& TFM = GetTextureFormatManagerRef();
 		const ITextureFormat* FormatObject = TFM.FindTextureFormat(BaseFormatName);
 
 		checkf(FormatObject != nullptr, TEXT("Bad FormatName %s passed to FChildTextureFormat::GetBaseFormatObject()"));
diff --git a/Engine/Source/Developer/TextureFormat/Public/TextureFormatManager.h b/Engine/Source/Developer/TextureFormat/Public/TextureFormatManager.h
--- a/Engine/Source/Developer/TextureFormat/Public/TextureFormatManager.h
+++ b/Engine/Source/Developer/TextureFormat/Public/TextureFormatManager.h
@@ -9,7 +9,14 @@
 /** Return the Texture Format Manager interface, if it is available, otherwise return nullptr. **/
 inline ITextureFormatManagerModule* GetTextureFormatManager()
 {
-	return FModuleManager::LoadModulePtr<ITextureFormatManagerModule>("TextureFormat");
+	// GetModule , not Load.  Should be loaded by TargetPlatform or Program startup on main thread.
+	//  cannot load modules from threads so must be done before any thread needs this module
+	ITextureFormatManagerModule* TFMM = FModuleManager::GetModulePtr<ITextureFormatManagerModule>("TextureFormat");
+	if ( TFMM == nullptr )
+	{
+		UE_LOG(LogInit, Error, TEXT("Texture format manager Ptr was requested but not available."));
+	}
+	return TFMM;
 }
 
 /** Return the Texture Format Manager interface, fatal error if it is not available. **/
@@ -18,7 +25,7 @@ inline ITextureFormatManagerModule& GetTextureFormatManagerRef()
 	class ITextureFormatManagerModule* TextureFormatManager = GetTextureFormatManager();
 	if (!TextureFormatManager)
 	{
-		UE_LOG(LogInit, Fatal, TEXT("Texture format manager was requested, but not available."));
+		UE_LOG(LogInit, Fatal, TEXT("Texture format manager Ref was requested, but not available."));
 		CA_ASSUME( TextureFormatManager != NULL );	// Suppress static analysis warning in unreachable code (fatal error)
 	}
 	return *TextureFormatManager;
