Clarifying comment

diff --git a/Engine/Source/Developer/TextureFormat/Public/Interfaces/ITextureFormat.h b/Engine/Source/Developer/TextureFormat/Public/Interfaces/ITextureFormat.h
--- a/Engine/Source/Developer/TextureFormat/Public/Interfaces/ITextureFormat.h
+++ b/Engine/Source/Developer/TextureFormat/Public/Interfaces/ITextureFormat.h
@@ -201,6 +201,7 @@ struct FEncodedTextureDescription
 		return Results;
 	}
 
+	// Returns the byte size of the unpacked/tiled mip. For mip chains that are packed or tiled, use FEncodedTextureExtendedData::MipSizesInBytes.
 	uint64 GetMipSizeInBytes(int32 InMipIndex) const
 	{
 		FIntVector3 MipDims = GetMipDimensions(InMipIndex);
