Fix for a malformed check() when assembling texture arrays in the editor one-by-one.

diff --git a/Engine/Source/Developer/TextureFormat/Public/Interfaces/ITextureFormat.h b/Engine/Source/Developer/TextureFormat/Public/Interfaces/ITextureFormat.h
--- a/Engine/Source/Developer/TextureFormat/Public/Interfaces/ITextureFormat.h
+++ b/Engine/Source/Developer/TextureFormat/Public/Interfaces/ITextureFormat.h
@@ -179,7 +179,7 @@ struct FEncodedTextureDescription
 			return 1; // no such thing as a cube volume, or a volume array.
 		}
 
-		check((bTextureArray && ArraySlices > 1) || (!bTextureArray && ArraySlices == 1));
+		check((bTextureArray && ArraySlices >= 1) || (!bTextureArray && ArraySlices == 1));
 		int32 Slices = ArraySlices;
 		if (bCubeMap)
 		{
