Fix for crash when creating texture arrays.

diff --git a/Engine/Source/Developer/TextureFormat/Public/Interfaces/ITextureFormat.h b/Engine/Source/Developer/TextureFormat/Public/Interfaces/ITextureFormat.h
--- a/Engine/Source/Developer/TextureFormat/Public/Interfaces/ITextureFormat.h
+++ b/Engine/Source/Developer/TextureFormat/Public/Interfaces/ITextureFormat.h
@@ -160,7 +160,7 @@ struct FEncodedTextureDescription
 			return FMath::Max(TopMipVolumeSizeZ >> InMipIndex, 1);
 		}
 
-		check ((bTextureArray && ArraySlices > 1) || (!bTextureArray && ArraySlices == 1));
+		check ((bTextureArray && ArraySlices >= 1) || (!bTextureArray && ArraySlices == 1));
 		int32 Slices = ArraySlices;
 		if (bCubeMap)
 		{
