fixed TSAN potential deadlock report by unlocking FTextureFormatManagerModule mutex while loading modules

diff --git a/Engine/Source/Developer/TextureFormat/Private/TextureFormatManagerModule.cpp b/Engine/Source/Developer/TextureFormat/Private/TextureFormatManagerModule.cpp
--- a/Engine/Source/Developer/TextureFormat/Private/TextureFormatManagerModule.cpp
+++ b/Engine/Source/Developer/TextureFormat/Private/TextureFormatManagerModule.cpp
@@ -84,23 +84,32 @@ public:
 			TArray<FTextureFormatMetadata> BaseModules;
 			TArray<FTextureFormatMetadata> ChildModules;
 
-			for (int32 Index = 0; Index < Modules.Num(); Index++)
 			{
-				if (Modules[Index] != ModuleName) // Avoid our own module when going through this list that was gathered by name
+				// unlock the mutex to avoid deadlock during module loading: T0 locks ModuleMutex (M0), loads the module and this broadcasts 
+				// FModuleManager::ModulesChangesEvent thread-safe delegate that locks its internal mutex (M1), while T1 loads a module ->
+				// broadcasts FModuleManager::ModulesChangesEvent (this locks M1) -> FTextureFormatManagerModule::ModulesChangesCallback that
+				// locks M0. at least it's what TSan reports. this can be a false positive, but there's no need to keep the mutex locked anyway, so it's
+				// better than just silencing TSan
+				FScopeUnlock ScopeUnlock(&ModuleMutex);
+
+				for (int32 Index = 0; Index < Modules.Num(); Index++)
 				{
-					ITextureFormatModule* Module = FModuleManager::LoadModulePtr<ITextureFormatModule>(Modules[Index]);
-					if (Module)
+					if (Modules[Index] != ModuleName) // Avoid our own module when going through this list that was gathered by name
 					{
-						FTextureFormatMetadata ModuleMeta;
-						ModuleMeta.Module = Module;
-						ModuleMeta.ModuleName = Modules[Index];
-						if ( Module->CanCallGetTextureFormats() )
+						ITextureFormatModule* Module = FModuleManager::LoadModulePtr<ITextureFormatModule>(Modules[Index]);
+						if (Module)
 						{
-							ChildModules.Add(ModuleMeta);
-						}
-						else
-						{
-							BaseModules.Add(ModuleMeta);
+							FTextureFormatMetadata ModuleMeta;
+							ModuleMeta.Module = Module;
+							ModuleMeta.ModuleName = Modules[Index];
+							if (Module->CanCallGetTextureFormats())
+							{
+								ChildModules.Add(ModuleMeta);
+							}
+							else
+							{
+								BaseModules.Add(ModuleMeta);
+							}
 						}
 					}
 				}
@@ -218,10 +227,13 @@ public:
 
 	virtual void Invalidate() override
 	{
-		FScopeLock Lock(&ModuleMutex);
-		// this is called from the constructor
-		TextureFormatsInitPhase = EInitPhase::Invalidated;
-		bForceCacheUpdate = true;
+		{
+			FScopeLock Lock(&ModuleMutex);
+			// this is called from the constructor
+			TextureFormatsInitPhase = EInitPhase::Invalidated;
+			bForceCacheUpdate = true;
+		}
+
 		GetTextureFormats();
 	}
 
