complimentary to 26518156, avoiding deadlock on module loading, handles one more case

diff --git a/Engine/Source/Developer/TextureFormat/Private/TextureFormatManagerModule.cpp b/Engine/Source/Developer/TextureFormat/Private/TextureFormatManagerModule.cpp
--- a/Engine/Source/Developer/TextureFormat/Private/TextureFormatManagerModule.cpp
+++ b/Engine/Source/Developer/TextureFormat/Private/TextureFormatManagerModule.cpp
@@ -227,6 +227,7 @@ public:
 
 	virtual void Invalidate() override
 	{
+		// don't lock `GetTextureFormats()` as it does own synchronisation
 		{
 			FScopeLock Lock(&ModuleMutex);
 			// this is called from the constructor
@@ -241,15 +242,21 @@ private:
 
 	void ModulesChangesCallback(FName InModuleName, EModuleChangeReason ReasonForChange)
 	{
-		FScopeLock Lock(&ModuleMutex);
-		if (bModuleChangeCallbackEnabled && (InModuleName != ModuleName) && InModuleName.ToString().Contains(TEXT("TextureFormat")))
+		// don't lock `Invalidate()` as it does own synchronisation
+		bool bLocalModuleChangeCallbackEnabled;
+		{
+			FScopeLock Lock(&ModuleMutex);
+			bLocalModuleChangeCallbackEnabled = bModuleChangeCallbackEnabled;
+		}
+
+		if (bLocalModuleChangeCallbackEnabled && (InModuleName != ModuleName) && InModuleName.ToString().Contains(TEXT("TextureFormat")))
 		{
 			// when a "TextureFormat" module is loaded, rebuild my list
 			Invalidate();
 		}
 	}
 
-	FName ModuleName;
+	const FName ModuleName;
 
 	TArray<const ITextureFormat*> TextureFormats;
 
