Fixed an issue with non pow 2 tiling on certain platforms, and in general cleaned things up. With the new platform preview I believe I've tested every possible combination of textures.

diff --git a/Engine/Source/Developer/TextureFormat/Public/Interfaces/ITextureFormat.h b/Engine/Source/Developer/TextureFormat/Public/Interfaces/ITextureFormat.h
--- a/Engine/Source/Developer/TextureFormat/Public/Interfaces/ITextureFormat.h
+++ b/Engine/Source/Developer/TextureFormat/Public/Interfaces/ITextureFormat.h
@@ -260,6 +260,13 @@ struct FEncodedTextureDescription
 		return GPixelFormats[PixelFormat].Get2DImageSizeInBytes(MipDims.X, MipDims.Y);
 	}
 
+	// Returns the bytes necessary to get to the next row of the current mip
+	uint64 GetMipSliceRowPitchBytes(int32 InMipIndex) const
+	{
+		const uint64 WidthInBlocks = GPixelFormats[PixelFormat].GetBlockCountForWidth(GetMipWidth(InMipIndex));
+		return WidthInBlocks * GPixelFormats[PixelFormat].BlockBytes;
+	}
+
 	int32 GetNumStreamingMips(const FEncodedTextureExtendedData* InExtendedData, const FTextureEngineParameters& InEngineParameters) const
 	{
 		return GetNumStreamingMipsDirect(NumMips, bCubeMap, bVolumeTexture, bTextureArray, InExtendedData, InEngineParameters);
