Clean up the texture format manager gathering function to clear up any possible mutex recursion.
#rnx

diff --git a/Engine/Source/Developer/TargetPlatform/Private/TargetPlatformManagerModule.cpp b/Engine/Source/Developer/TargetPlatform/Private/TargetPlatformManagerModule.cpp
--- a/Engine/Source/Developer/TargetPlatform/Private/TargetPlatformManagerModule.cpp
+++ b/Engine/Source/Developer/TargetPlatform/Private/TargetPlatformManagerModule.cpp
@@ -192,7 +192,6 @@ public:
 		, bForceCacheUpdate(true)
 		, bHasInitErrors(false)
 		, bIgnoreFirstDelegateCall(true)
-		, bSkipOneTextureFormatManagerInvalidate(false)
 	{
 #if WITH_EDITOR && UE_WITH_TURNKEY_SUPPORT
 
@@ -225,9 +224,6 @@ public:
 
 		TextureFormatManager = FModuleManager::LoadModulePtr<ITextureFormatManagerModule>("TextureFormat");
 
-		//TextureFormatManager->Invalidate() already done, don't do again now :
-		bSkipOneTextureFormatManagerInvalidate = true;
-
 		// Calling a virtual function from a constructor, but with no expectation that a derived implementation of this
 		// method would be called.  This is solely to avoid duplicating code in this implementation, not for polymorphism.
 		FTargetPlatformManagerModule::Invalidate();
@@ -264,15 +260,6 @@ public:
 
 		bForceCacheUpdate = false;
 		
-		if ( bSkipOneTextureFormatManagerInvalidate )
-		{
-			bSkipOneTextureFormatManagerInvalidate = false;
-		}
-		else if (!bHasInitErrors)
-		{
-			TextureFormatManager->Invalidate();
-		}
-
 		// If we've had an error due to an invalid target platform, don't do additional work
 		if (!bHasInitErrors)
 		{
@@ -688,12 +675,6 @@ public:
 		return nullptr;
 	}
 
-	virtual const TArray<const ITextureFormat*>& GetTextureFormats() override
-	{
-		// note that this gets ALL ITextureFormat Modules, not just ones relevant to the current TargetPlatform
-		return TextureFormatManager->GetTextureFormats();
-	}
-
 	virtual const ITextureFormat* FindTextureFormat(FName Name) override
 	{
 		return TextureFormatManager->FindTextureFormat(Name);
@@ -1363,9 +1344,6 @@ private:
 	// Flag to avoid redunant reloads
 	bool bIgnoreFirstDelegateCall;
 	
-	// Flag to avoid redunant reloads
-	bool bSkipOneTextureFormatManagerInvalidate;
-
 	// Holds the list of discovered platforms.
 	TArray<ITargetPlatform*> Platforms;
 	TArray<ITargetPlatformControls*> PlatformControls;
diff --git a/Engine/Source/Developer/TargetPlatform/Public/Interfaces/ITargetPlatformManagerModule.h b/Engine/Source/Developer/TargetPlatform/Public/Interfaces/ITargetPlatformManagerModule.h
--- a/Engine/Source/Developer/TargetPlatform/Public/Interfaces/ITargetPlatformManagerModule.h
+++ b/Engine/Source/Developer/TargetPlatform/Public/Interfaces/ITargetPlatformManagerModule.h
@@ -154,7 +154,8 @@ public:
 	 *
 	 * @return Collection of texture formats.
 	 */
-	virtual const TArray<const class ITextureFormat*>& GetTextureFormats() = 0;
+	UE_DEPRECATED(5.6, "Not thread safe, use FindTextureFormat")
+	TArray<const class ITextureFormat*> GetTextureFormats() { return {}; }
 
 	/**
 	 * Determine if there were errors during the initialization of the platform manager.
diff --git a/Engine/Source/Developer/TextureFormat/Private/TextureFormatManagerModule.cpp b/Engine/Source/Developer/TextureFormat/Private/TextureFormatManagerModule.cpp
--- a/Engine/Source/Developer/TextureFormat/Private/TextureFormatManagerModule.cpp
+++ b/Engine/Source/Developer/TextureFormat/Private/TextureFormatManagerModule.cpp
@@ -28,18 +28,14 @@ public:
 		GetTextureFormatsDone = 4
 	};
 
-	/** Default constructor. */
 	FTextureFormatManagerModule()
 		: ModuleName(TEXT("TextureFormat"))
 		, bForceCacheUpdate(true)
-		, bModuleChangeCallbackEnabled(false)
 		, TextureFormatsInitPhase(EInitPhase::JustConstructedNotInit)
 	{
-		// Calling a virtual function from a constructor, but with no expectation that a derived implementation of this
-		// method would be called.  This is solely to avoid duplicating code in this implementation, not for polymorphism.
-		FTextureFormatManagerModule::Invalidate();
-		
-		// add AFTER Invalidate :
+		UpdateTextureFormatList();
+
+		// For tracking texture format discovery.
 		FModuleManager::Get().OnModulesChanged().AddRaw(this, &FTextureFormatManagerModule::ModulesChangesCallback);
 	}
 
@@ -51,7 +47,7 @@ public:
 		FModuleManager::Get().OnModulesChanged().RemoveAll(this);
 	}
 
-	virtual const TArray<const ITextureFormat*>& GetTextureFormats() override
+	void UpdateTextureFormatList()
 	{
 		FScopeLock Lock(&ModuleMutex);
 
@@ -63,7 +59,6 @@ public:
 		{
 			// turn off flag immediately so that repeated calls to GetTextureFormats will not come in here again
 			bForceCacheUpdate = false;
-			bModuleChangeCallbackEnabled = false; // don't re-call me from my own module loads
 			TextureFormatsInitPhase = EInitPhase::GetTextureFormatsInProgressDontTouch;
 
 			// note the first time this is done is from FTargetPlatformManagerModule::FTargetPlatformManagerModule()
@@ -79,44 +74,36 @@ public:
 			if (!Modules.Num())
 			{
 				UE_LOG(LogTextureFormatManager, Error, TEXT("No texture formats found!"));
-			}			
-					
+			}
+			
+			// This is all because child formats will do a LoadModule on base formats and expect them to be ready, so
+			// we make sure they are done first.
 			TArray<FTextureFormatMetadata> BaseModules;
 			TArray<FTextureFormatMetadata> ChildModules;
 
+			for (int32 Index = 0; Index < Modules.Num(); Index++)
 			{
-				// unlock the mutex to avoid deadlock during module loading: T0 locks ModuleMutex (M0), loads the module and this broadcasts 
-				// FModuleManager::ModulesChangesEvent thread-safe delegate that locks its internal mutex (M1), while T1 loads a module ->
-				// broadcasts FModuleManager::ModulesChangesEvent (this locks M1) -> FTextureFormatManagerModule::ModulesChangesCallback that
-				// locks M0. at least it's what TSan reports. this can be a false positive, but there's no need to keep the mutex locked anyway, so it's
-				// better than just silencing TSan
-				FScopeUnlock ScopeUnlock(&ModuleMutex);
-
-				for (int32 Index = 0; Index < Modules.Num(); Index++)
+				if (Modules[Index] != ModuleName) // Avoid our own module when going through this list that was gathered by name
 				{
-					if (Modules[Index] != ModuleName) // Avoid our own module when going through this list that was gathered by name
+					ITextureFormatModule* Module = FModuleManager::LoadModulePtr<ITextureFormatModule>(Modules[Index]);
+					if (Module)
 					{
-						ITextureFormatModule* Module = FModuleManager::LoadModulePtr<ITextureFormatModule>(Modules[Index]);
-						if (Module)
+						FTextureFormatMetadata ModuleMeta;
+						ModuleMeta.Module = Module;
+						ModuleMeta.ModuleName = Modules[Index];
+						if (Module->CanCallGetTextureFormats()) // child modules want to call GetTextureFormats
 						{
-							FTextureFormatMetadata ModuleMeta;
-							ModuleMeta.Module = Module;
-							ModuleMeta.ModuleName = Modules[Index];
-							if (Module->CanCallGetTextureFormats())
-							{
-								ChildModules.Add(ModuleMeta);
-							}
-							else
-							{
-								BaseModules.Add(ModuleMeta);
-							}
+							ChildModules.Add(ModuleMeta);
+						}
+						else
+						{
+							BaseModules.Add(ModuleMeta);
 						}
 					}
 				}
 			}
 			
 			// first populate TextureFormats[] with all Base Modules
-			// 
 			for (int32 Index = 0; Index < BaseModules.Num(); Index++)
 			{
 				ITextureFormatModule* Module = BaseModules[Index].Module;
@@ -124,7 +111,6 @@ public:
 				ITextureFormat* Format = Module->GetTextureFormat();
 				if (Format != nullptr)
 				{
-				
 					// I want to see this log by default in Cook+Editor , but not in TBW
 					#ifndef VerboseIfNotEditor
 					#if WITH_EDITOR
@@ -146,6 +132,7 @@ public:
 
 			// run through the Child formats and call GetTextureFormat() on them
 			// this could call back to me and do GetTextureFormats() which will get only the base formats
+			TArray<TPair<ITextureFormat*, int32>, TInlineAllocator<32>> ReadyChildModules;
 			for (int32 Index = 0; Index < ChildModules.Num(); Index++)
 			{
 				ITextureFormatModule* Module = ChildModules[Index].Module;
@@ -155,35 +142,22 @@ public:
 				{
 					UE_LOG(LogTextureFormatManager,VerboseIfNotEditor,TEXT("Loaded Child TextureFormat: %s"),*ChildModules[Index].ModuleName.ToString());
 
-					// do not add me to TextureFormats yet
+					// do not add me to TextureFormats yet'
+					ReadyChildModules.Add({Format, Index});
 				}
 			}
-			
-			// back up phase to 2, no calls to GetTextureFormats() allowed now
-			TextureFormatsInitPhase = EInitPhase::GetTextureFormatsInProgressDontTouch;
 
-			for (int32 Index = 0; Index < ChildModules.Num(); Index++)
+			for (TPair<ITextureFormat*, int32>& ReadyChild : ReadyChildModules)
 			{
-				ITextureFormatModule* Module = ChildModules[Index].Module;
-
-				// GetTextureFormat was already done so this should just return a stored pointer, no more init
-				ITextureFormat* Format = Module->GetTextureFormat();
-				if (Format != nullptr)
-				{
-					// now add to the list :
-					TextureFormats.Add(Format);
-					TextureFormatMetadata.Add(ChildModules[Index]);
-				}
+				TextureFormats.Add(ReadyChild.Key);
+				TextureFormatMetadata.Add(ChildModules[ReadyChild.Value]);
 			}
-
+			
 			// all done :
 			TextureFormatsInitPhase = EInitPhase::GetTextureFormatsDone;
-			bModuleChangeCallbackEnabled = true;
 		}
 
 		check( (int)TextureFormatsInitPhase >= (int)EInitPhase::GetTextureFormatsPartialOkayToRead );
-
-		return TextureFormats;
 	}
 	
 	virtual const ITextureFormat* FindTextureFormat(FName Name) override
@@ -199,9 +173,6 @@ public:
 		FScopeLock Lock(&ModuleMutex);
 		check( (int)TextureFormatsInitPhase >= (int)EInitPhase::GetTextureFormatsPartialOkayToRead );
 
-		// Called to ensure the arrays are populated
-		// dangerous and not necessary, removed :
-		//GetTextureFormats();
 		check( ! bForceCacheUpdate );
 
 		for (int32 Index = 0; Index < TextureFormats.Num(); Index++)
@@ -225,35 +196,39 @@ public:
 		return nullptr;
 	}
 
-	virtual void Invalidate() override
-	{
-		// don't lock `GetTextureFormats()` as it does own synchronisation
-		{
-			FScopeLock Lock(&ModuleMutex);
-			// this is called from the constructor
-			TextureFormatsInitPhase = EInitPhase::Invalidated;
-			bForceCacheUpdate = true;
-		}
-
-		GetTextureFormats();
-	}
 
 private:
 
 	void ModulesChangesCallback(FName InModuleName, EModuleChangeReason ReasonForChange)
 	{
-		// don't lock `Invalidate()` as it does own synchronisation
-		bool bLocalModuleChangeCallbackEnabled;
+		//
+		// This is complex because this is the only place we can set up our texture format list
+		// from the game thread. The only time we can update "on demand" could be from any thread,
+		// which prevents up from calling LoadModule.
+		//
+		// However, this gets called while we are loading our modules.
+		//
+		// In order to avoid recursion, we only do LoadModules in response to module _discovery_, 
+		// thus ensuring we can't recurse from our Invalidate call.
+		//
+		//
+		if (ReasonForChange != EModuleChangeReason::PluginDirectoryChanged || // only care about discovery
+			(InModuleName == ModuleName) ||  // don't care about _this_ module
+			!InModuleName.ToString().Contains(TEXT("TextureFormat"))) // only care about texture formats.
 		{
-			FScopeLock Lock(&ModuleMutex);
-			bLocalModuleChangeCallbackEnabled = bModuleChangeCallbackEnabled;
+			return;
 		}
 
-		if (bLocalModuleChangeCallbackEnabled && (InModuleName != ModuleName) && InModuleName.ToString().Contains(TEXT("TextureFormat")))
-		{
-			// when a "TextureFormat" module is loaded, rebuild my list
-			Invalidate();
-		}
+		// when a "TextureFormat" module is discovered, rebuild my list.
+
+		// Note... it's unclear but it looks like it _might_ be possible for a LoadModule to 
+		// cause a module that it loads to get discovered, which could cause a recursion. However
+		// texture format modules are pretty straightforward and shouldn't ever get here as they
+		// are discovered on startup.
+
+		// In order to even get here you have to put a texture format in a plugin that gets loaded after
+		// startup. (I think)
+		UpdateTextureFormatList();
 	}
 
 	const FName ModuleName;
@@ -271,9 +246,6 @@ private:
 	// in case of a module reload of a TextureFormat-Module.
 	bool bForceCacheUpdate;
 
-	// Flag to avoid redunant reloads
-	bool bModuleChangeCallbackEnabled;
-
 	// Track tricky initialization progress
 	EInitPhase TextureFormatsInitPhase;
 
diff --git a/Engine/Source/Developer/TextureFormat/Public/Interfaces/ITextureFormatManagerModule.h b/Engine/Source/Developer/TextureFormat/Public/Interfaces/ITextureFormatManagerModule.h
--- a/Engine/Source/Developer/TextureFormat/Public/Interfaces/ITextureFormatManagerModule.h
+++ b/Engine/Source/Developer/TextureFormat/Public/Interfaces/ITextureFormatManagerModule.h
@@ -14,7 +14,7 @@ class ITextureFormatManagerModule
 public:
 
 	/**
-	 * Finds a texture format with the specified name.
+	 * Finds a texture format with the specified name. Safe to call from any thread.
 	 *
 	 * @param Name Name of the format to find.
 	 * @return The texture format, or nullptr if not found.
@@ -23,7 +23,8 @@ public:
 
 	/**
 	 * Finds a texture format with the specified name and provides information about the module it came from.
-	 *
+	 * Safe to call from any thread.
+	 * 
 	 * @param Name Name of the format to find.
 	 * @param OutModuleName Name of the module that the found format came from, or unmodified if not found.
 	 * @param OutModule Interface of the module that the found format came from, or unmodified if not found.
@@ -36,15 +37,17 @@ public:
 	 *
 	 * @return Collection of texture formats.
 	 */
-	virtual const TArray<const class ITextureFormat*>& GetTextureFormats() = 0;
+	UE_DEPRECATED(5.6, "Deprecated as it's not thread safe - use FindTextureFormat")
+	TArray<const class ITextureFormat*> GetTextureFormats()	{ return {}; }
 
 	/**
 	 * Invalidates the texture format manager module.
-	 *
-	 * Invalidate should be called if any TextureFormat modules get loaded/unloaded/reloaded during 
-	 * runtime to give the implementation the chance to rebuild all its internal states and caches.
+	 * 
+	 * This is no longer necessary as all work is done in response to broadcast plugin/module discovery messages.
 	 */
-	virtual void Invalidate() = 0;
+	UE_DEPRECATED(5.6, "No longer necessary to call")
+	void Invalidate() { }
+
 public:
 
 	/** Virtual destructor. */
