Deprecating the "Tiling" style of texture formats.

diff --git a/Engine/Source/Developer/TextureFormat/Public/ChildTextureFormat.h b/Engine/Source/Developer/TextureFormat/Public/ChildTextureFormat.h
--- a/Engine/Source/Developer/TextureFormat/Public/ChildTextureFormat.h
+++ b/Engine/Source/Developer/TextureFormat/Public/ChildTextureFormat.h
@@ -229,7 +229,7 @@ public:
 		}
 		return true;
 	}
-
+	PRAGMA_DISABLE_DEPRECATION_WARNINGS
 	bool CompressBaseImageTiled(
 		const FImage* Images,
 		uint32 NumImages,
@@ -283,6 +283,7 @@ public:
 
 		return GetBaseFormatObject(BuildSettings.TextureFormatName)->ReleaseTiling(BuildSettings, TilerSettings);
 	}
+	PRAGMA_ENABLE_DEPRECATION_WARNINGS
 
 
 	virtual FCbObject ExportGlobalFormatConfig(const FTextureBuildSettings& BuildSettings) const override
diff --git a/Engine/Source/Developer/TextureFormat/Public/Interfaces/ITextureFormat.h b/Engine/Source/Developer/TextureFormat/Public/Interfaces/ITextureFormat.h
--- a/Engine/Source/Developer/TextureFormat/Public/Interfaces/ITextureFormat.h
+++ b/Engine/Source/Developer/TextureFormat/Public/Interfaces/ITextureFormat.h
@@ -602,6 +602,7 @@ public:
 	 * An object produced by PrepareTiling and used by SetTiling and CompressImageTiled.
 	 * This is used as an inheritance base for tiling formats to add their own information.
 	 */
+	
 	struct FTilerSettings
 	{
 	};
@@ -617,6 +618,7 @@ public:
 	 * @param Tiler The tiler settings.
 	 * @returns true on success, false otherwise.
 	 */
+	UE_DEPRECATED(5.6, "Tiling compressors will no longer be supported; compress the mip entirely via CompressImage(Ex) and UE will tile the image afterwards")
 	virtual bool CompressImageTiled(
 		const FImage* Images,
 		uint32 NumImages,
@@ -637,6 +639,7 @@ public:
 	 * @returns true if tiling is supported, false if it must be done by the caller
 	 *
 	 */
+	UE_DEPRECATED(5.6, "Tiling compressors will no longer be supported; compress the mip entirely via CompressImage(Ex) and UE will tile the image afterwards")
 	virtual bool SupportsTiling() const
 	{
 		return false;
@@ -652,6 +655,7 @@ public:
 	 * @param OutCompressedImage The image to tile.
 	 * @returns true on success, false otherwise.
 	 */
+	UE_DEPRECATED(5.6, "Tiling compressors will no longer be supported; compress the mip entirely via CompressImage(Ex) and UE will tile the image afterwards")
 	virtual bool PrepareTiling(
 		const FImage* Images,
 		const uint32 NumImages,
@@ -674,6 +678,7 @@ public:
 	 * @param NumBlocks The number of blocks.
 	 * @returns true on success, false otherwise.
 	 */
+	UE_DEPRECATED(5.6, "Tiling compressors will no longer be supported; compress the mip entirely via CompressImage(Ex) and UE will tile the image afterwards")
 	virtual bool SetTiling(
 		const FTextureBuildSettings& BuildSettings,
 		TSharedPtr<FTilerSettings>& TilerSettings,
@@ -691,6 +696,7 @@ public:
 	 * @param BuildSettings Build settings.
 	 * @param TilerSettings The tiler settings object to release.
 	 */
+	UE_DEPRECATED(5.6, "Tiling compressors will no longer be supported; compress the mip entirely via CompressImage(Ex) and UE will tile the image afterwards")
 	virtual void ReleaseTiling(const FTextureBuildSettings& BuildSettings, TSharedPtr<FTilerSettings>& TilerSettings) const
 	{
 		unimplemented();
diff --git a/Engine/Source/Runtime/Engine/Private/TextureDerivedData.cpp b/Engine/Source/Runtime/Engine/Private/TextureDerivedData.cpp
--- a/Engine/Source/Runtime/Engine/Private/TextureDerivedData.cpp
+++ b/Engine/Source/Runtime/Engine/Private/TextureDerivedData.cpp
@@ -990,10 +990,12 @@ static void FinalizeBuildSettingsForLayer(
 
 			// Shared linear encoding can only work if the base texture format does not expect to
 			// do the tiling itself (SupportsTiling == false).
+			PRAGMA_DISABLE_DEPRECATION_WARNINGS
 			if (ChildTextureFormat && OutSettings.BaseTextureFormat->SupportsTiling() == false)
 			{
 				OutSettings.Tiler = ChildTextureFormat->GetTiler();
 			}
+			PRAGMA_ENABLE_DEPRECATION_WARNINGS
 		}
 	}
 
