TextureExtendedData now always holds mip sizes and knows when the data has been tiled.

#rnx

diff --git a/Engine/Source/Developer/TextureBuild/Private/TextureBuildFunction.cpp b/Engine/Source/Developer/TextureBuild/Private/TextureBuildFunction.cpp
--- a/Engine/Source/Developer/TextureBuild/Private/TextureBuildFunction.cpp
+++ b/Engine/Source/Developer/TextureBuild/Private/TextureBuildFunction.cpp
@@ -396,16 +396,9 @@ static bool ReadChildBuildInputs(FChildBuildData& OutChildBuildInputs, UE::Deriv
 	OutChildBuildInputs.NumStreamingMips = OutChildBuildInputs.TextureDescription.GetNumStreamingMips(&OutChildBuildInputs.TextureExtendedData, OutChildBuildInputs.EngineParameters);
 	OutChildBuildInputs.NumEncodedMips = OutChildBuildInputs.TextureDescription.GetNumEncodedMips(&OutChildBuildInputs.TextureExtendedData);
 
-	if (OutChildBuildInputs.TextureExtendedData.MipSizesInBytes.Num() == 0)
-	{
-		// Init with linear sizes.
-		OutChildBuildInputs.TextureExtendedData.MipSizesInBytes.SetNumUninitialized(OutChildBuildInputs.NumEncodedMips);
-		for (int32 MipIndex = 0; MipIndex < OutChildBuildInputs.NumEncodedMips; MipIndex++)
-		{
-			OutChildBuildInputs.TextureExtendedData.MipSizesInBytes[MipIndex] = OutChildBuildInputs.TextureDescription.GetMipSizeInBytes(MipIndex);
-		}
-	}
-
+	// Extended data mip sizes should always be valid with either linear mip sizes or tiled mip sizes.
+	check(OutChildBuildInputs.TextureExtendedData.MipSizesInBytes.Num() == OutChildBuildInputs.TextureDescription.NumMips);
+	
 	{
 		FSharedBuffer InputTextureMipTailData;
 		if (OutChildBuildInputs.TextureDescription.NumMips > OutChildBuildInputs.NumStreamingMips)
diff --git a/Engine/Source/Developer/TextureFormat/Public/Interfaces/ITextureFormat.h b/Engine/Source/Developer/TextureFormat/Public/Interfaces/ITextureFormat.h
--- a/Engine/Source/Developer/TextureFormat/Public/Interfaces/ITextureFormat.h
+++ b/Engine/Source/Developer/TextureFormat/Public/Interfaces/ITextureFormat.h
@@ -64,24 +64,37 @@ static bool GetStreamingDisabledForNonVirtualTextureProperties(bool bInCubeMap,
 
 
 
-// Extra data for an encoded texture. For "normal" textures (i.e. linear, without a packed mip tail), this must be
-// all zeroes.
+// Extra data for an encoded texture.
+// This is filled out by platform "child" formats that can potentially tile or otherwise
+// reorganize the mip data for use on the target platform. For platforms that use bog standard
+// "linear" texture data, this is mostly zeroes, but must still be "Valid" data returned by the
+// GetExtendedDataForTexture - notably the MipSizesInBytes.
 struct FEncodedTextureExtendedData
 {
 	// Copied from RHIDefinitions.h - we expose here to avoid the include.
 	static constexpr int32 MAX_TEXTURE_MIP_COUNT=15;
 
+	// Some platforms group the smallest mips in to a single bulk chunk of data. This is nonzero
+	// in such cases and the MipSizesInBytes array is truncated.
 	int32 NumMipsInTail = 0;
+
+	// ExtData is metadata the tiler passes through to the runtime.
 	uint32 ExtData = 0;
 	
+	// If true, the mip data can't be assumed to be linear.
+	bool bIsTiled = false;
+
 	// If true, this texture might change layouts if top mips are striped (i.e. LODBias is not zero).
 	bool bSensitiveToLODBias = false;
 	
 	// If bSensitiveToLODBias is set, this is the LODBias for this layout.
 	int8 LODBiasIfSensitive = 0;
 
-	// With packing/tiling, mip sizes are not trivially computable. Not that these sizes must NOT be
-	// used for mips prior to tiling. For those, FEncodedTextureDescription::GetMipSizeInBytes().
+	// All generators of this structure are required to fill this structure out even if they are
+	// default linear sizes. For non linear textures these are not trivially calculatable.
+	// Also note that if the texture is sensitive to LODBias with respect to tiling, the indeces in this
+	// array that are for mips that would be stripped are present but should be ignored as they don't 
+	// represent anything we can accurately calculate.
 	TArray<uint64, TInlineAllocator<MAX_TEXTURE_MIP_COUNT>> MipSizesInBytes;
 };
 
@@ -518,9 +531,15 @@ public:
 	* for textures that have been transformed in some way for a platform. LODBias is needed because in some cases the tiling
 	* changes based on the top mip actually given to the hardware.
 	*/
-	virtual FEncodedTextureExtendedData GetExtendedDataForTexture(const FEncodedTextureDescription& InTextureDescription, int8 InLODBias) const
+	[[nodiscard]] virtual FEncodedTextureExtendedData GetExtendedDataForTexture(const FEncodedTextureDescription& InTextureDescription, int8 InLODBias) const
 	{
-		return FEncodedTextureExtendedData();
+		FEncodedTextureExtendedData ExtendedData;
+		ExtendedData.MipSizesInBytes.SetNumUninitialized(InTextureDescription.NumMips);
+		for (int32 MipIndex = 0; MipIndex < InTextureDescription.NumMips; MipIndex++)
+		{
+			ExtendedData.MipSizesInBytes[MipIndex] = InTextureDescription.GetMipSizeInBytes(MipIndex);
+		}
+		return ExtendedData;
 	}
 
 	// Return true if this format can decode the given pixel format to one of the ERawImageFormats.
diff --git a/Engine/Source/Runtime/Engine/Private/TextureDerivedData.cpp b/Engine/Source/Runtime/Engine/Private/TextureDerivedData.cpp
--- a/Engine/Source/Runtime/Engine/Private/TextureDerivedData.cpp
+++ b/Engine/Source/Runtime/Engine/Private/TextureDerivedData.cpp
@@ -695,7 +695,7 @@ static void GetBuiltTextureSizeBytesEstimate(
 
 		int32 LODBias = 0;
 		FEncodedTextureExtendedData ExtendedData = TextureFormat->GetExtendedDataForTexture(TextureDescription, LODBias);
-		if ( ExtendedData.MipSizesInBytes.Num() > 0 )
+		if ( ExtendedData.bIsTiled )
 		{
 			// ExtendedData is only valid for platform/tiled images
 		
diff --git a/Engine/Source/Runtime/Engine/Private/TextureDerivedDataTask.cpp b/Engine/Source/Runtime/Engine/Private/TextureDerivedDataTask.cpp
--- a/Engine/Source/Runtime/Engine/Private/TextureDerivedDataTask.cpp
+++ b/Engine/Source/Runtime/Engine/Private/TextureDerivedDataTask.cpp
@@ -2267,29 +2267,6 @@ static bool UnpackPlatformDataFromBuild(FTexturePlatformData& OutPlatformData, U
 	int32 NumEncodedMips = EncodedTextureDescription.GetNumEncodedMips(&EncodedTextureExtendedData);
 	check(NumEncodedMips >= NumStreamingMips);
 
-	//
-	// Mips are split up:
-	//	Streaming mips are all stored independently under value FTexturePlatformData::MakeMipId(MipIndex);
-	//	Nonstreaming ("inlined") mips are stored in one buffer under value "MipTail". To disentangle the separate mips
-	//	we need their size. If EncodedTextureExtendedData is present, we have the ::MipSizesInBytes field. If that isn't
-	//	present, then we know it's nothing fancy and we can directly compute mip sizes.
-	//
-	TArray<uint64, TInlineAllocator<MAX_TEXTURE_MIP_COUNT>> MipSizesFromBuild;
-	if (EncodedTextureExtendedData.MipSizesInBytes.Num() == 0)
-	{
-		// Normal mip sizes - generate from texture description
-		MipSizesFromBuild.AddUninitialized(EncodedTextureDescription.NumMips);
-		for (int32 MipIndex = 0; MipIndex < EncodedTextureDescription.NumMips; MipIndex++)
-		{
-			MipSizesFromBuild[MipIndex] = EncodedTextureDescription.GetMipSizeInBytes(MipIndex);
-		}
-	}
-
-	// For the moment, we don't necessarily have the sizes from the extended data, in that case we must have them from the build.
-	TArray<uint64, TInlineAllocator<MAX_TEXTURE_MIP_COUNT>>& MipSizesInBytes = EncodedTextureExtendedData.MipSizesInBytes.Num() ? EncodedTextureExtendedData.MipSizesInBytes : MipSizesFromBuild;
-	check(MipSizesInBytes.Num() == NumEncodedMips);
-
-
 	//
 	//
 	// We have all the metadata we need, we can grab the data
@@ -2322,6 +2299,13 @@ static bool UnpackPlatformDataFromBuild(FTexturePlatformData& OutPlatformData, U
 
 
 
+	//
+	// Mips are split up:
+	//	Streaming mips are all stored independently under value FTexturePlatformData::MakeMipId(MipIndex);
+	//	Nonstreaming ("inlined") mips are stored in one buffer under value "MipTail". To disentangle the separate mips
+	//	we need their size.
+	//
+
 	uint64 CurrentMipTailOffset = 0;
 	for (int32 MipIndex = 0; MipIndex < EncodedTextureDescription.NumMips; MipIndex++)
 	{
@@ -2345,7 +2329,7 @@ static bool UnpackPlatformDataFromBuild(FTexturePlatformData& OutPlatformData, U
 			// coincidentally be the same. All mips past NumStreamingMips need to be copied into the bulk data
 			// and are always resident in memory with the texture.
 
-			uint64 MipSizeInBytes = MipSizesInBytes[MipIndex];
+			uint64 MipSizeInBytes = EncodedTextureExtendedData.MipSizesInBytes[MipIndex];
 			FMemoryView MipView = MipTailData.GetView().Mid(CurrentMipTailOffset, MipSizeInBytes);
 			CurrentMipTailOffset += MipSizeInBytes;
 
